Added Thumbstick Right Y-Axis to Thrust control so virtual joystick works.

diff --git a/Templates/TP_Flying/Config/DefaultInput.ini b/Templates/TP_Flying/Config/DefaultInput.ini
--- a/Templates/TP_Flying/Config/DefaultInput.ini
+++ b/Templates/TP_Flying/Config/DefaultInput.ini
@@ -40,6 +40,7 @@
 +AxisConfig=(AxisKeyName="OculusTouch_Right_IndexPointing",AxisProperties=(DeadZone=0.000000,Sensitivity=1.000000,Exponent=1.000000,bInvert=False))
 +AxisConfig=(AxisKeyName="OculusTouch_Right_ThumbUp",AxisProperties=(DeadZone=0.000000,Sensitivity=1.000000,Exponent=1.000000,bInvert=False))
 bAltEnterTogglesFullscreen=True
+bF11TogglesFullscreen=True
 bUseMouseForTouch=False
 bEnableMouseSmoothing=True
 bEnableFOVScaling=True
@@ -76,6 +77,7 @@ DefaultViewportMouseLockMode=LockOnCapture
 +AxisMappings=(AxisName="Thrust",Key=MotionController_Left_Trigger,Scale=-1.000000)
 +AxisMappings=(AxisName="Thrust",Key=Gamepad_RightStick_Up,Scale=1.000000)
 +AxisMappings=(AxisName="Thrust",Key=Gamepad_RightStick_Down,Scale=-1.000000)
++AxisMappings=(AxisName="Thrust",Key=Gamepad_RightY,Scale=-1.000000)
 bAlwaysShowTouchInterface=False
 bShowConsoleOnFourFingerTap=True
 DefaultTouchInterface=/Engine/MobileResources/HUD/DefaultVirtualJoysticks.DefaultVirtualJoysticks
diff --git a/Templates/TP_FlyingBP/Config/DefaultInput.ini b/Templates/TP_FlyingBP/Config/DefaultInput.ini
--- a/Templates/TP_FlyingBP/Config/DefaultInput.ini
+++ b/Templates/TP_FlyingBP/Config/DefaultInput.ini
@@ -40,6 +40,7 @@
 +AxisConfig=(AxisKeyName="OculusTouch_Right_IndexPointing",AxisProperties=(DeadZone=0.000000,Sensitivity=1.000000,Exponent=1.000000,bInvert=False))
 +AxisConfig=(AxisKeyName="OculusTouch_Right_ThumbUp",AxisProperties=(DeadZone=0.000000,Sensitivity=1.000000,Exponent=1.000000,bInvert=False))
 bAltEnterTogglesFullscreen=True
+bF11TogglesFullscreen=True
 bUseMouseForTouch=False
 bEnableMouseSmoothing=True
 bEnableFOVScaling=True
@@ -74,8 +75,10 @@ DefaultViewportMouseLockMode=LockOnCapture
 +AxisMappings=(AxisName="MoveRight",Key=MotionController_Right_Thumbstick_X,Scale=1.000000)
 +AxisMappings=(AxisName="Thrust",Key=MotionController_Right_Trigger,Scale=1.000000)
 +AxisMappings=(AxisName="Thrust",Key=MotionController_Left_Trigger,Scale=-1.000000)
++AxisMappings=(AxisName="Thrust",Key=Gamepad_RightY,Scale=-1.000000)
 +AxisMappings=(AxisName="Thrust",Key=Gamepad_RightStick_Up,Scale=1.000000)
 +AxisMappings=(AxisName="Thrust",Key=Gamepad_RightStick_Down,Scale=-1.000000)
++AxisMappings=(AxisName="Thrust",Key=Gamepad_RightY,Scale=-1.000000)
 bAlwaysShowTouchInterface=False
 bShowConsoleOnFourFingerTap=True
 DefaultTouchInterface=/Engine/MobileResources/HUD/DefaultVirtualJoysticks.DefaultVirtualJoysticks
